#UE4 #HTML5  TTP: 50534 Audio broken for HTML5 in main.

       - catch up with changes in core Audio ( add relevant virtual functions ).
       - clean up logging.

diff --git a/Engine/Source/Runtime/HTML5/HTML5Audio/Private/HTML5AudioBuffer.cpp b/Engine/Source/Runtime/HTML5/HTML5Audio/Private/HTML5AudioBuffer.cpp
--- a/Engine/Source/Runtime/HTML5/HTML5Audio/Private/HTML5AudioBuffer.cpp
+++ b/Engine/Source/Runtime/HTML5/HTML5Audio/Private/HTML5AudioBuffer.cpp
@@ -9,7 +9,7 @@
 /*------------------------------------------------------------------------------------
 	FALSoundBuffer.
 ------------------------------------------------------------------------------------*/
-/** 
+/**
  * Constructor
  *
  * @param AudioDevice	audio device this sound buffer is going to be attached to.
@@ -20,8 +20,8 @@ FALSoundBuffer::FALSoundBuffer( FALAudioDevice* InAudioDevice )
 }
 
 /**
- * Destructor 
- * 
+ * Destructor
+ *
  * Frees wave data and detaches itself from audio device.
  */
 FALSoundBuffer::~FALSoundBuffer( void )
@@ -92,6 +92,18 @@ FALSoundBuffer* FALSoundBuffer::CreateNativeBuffer( FALAudioDevice* AudioDevice,
 {
 	SCOPE_CYCLE_COUNTER( STAT_AudioResourceCreationTime );
 
+	// This code is not relevant for now on HTML5 but adding this for consistency with other platforms.
+	// Check to see if thread has finished decompressing on the other thread
+
+	if (Wave->AudioDecompressor != NULL)
+	{
+		Wave->AudioDecompressor->EnsureCompletion();
+
+		// Remove the decompressor
+		delete Wave->AudioDecompressor;
+		Wave->AudioDecompressor = NULL;
+	}
+
 	// Can't create a buffer without any source data
 	if( Wave == NULL || Wave->NumChannels == 0 )
 	{
@@ -102,7 +114,7 @@ FALSoundBuffer* FALSoundBuffer::CreateNativeBuffer( FALAudioDevice* AudioDevice,
 	Wave->InitAudioResource(AudioDevice->GetRuntimeFormat());
 
 	FALSoundBuffer* Buffer = NULL;
-	
+
 	// Find the existing buffer if any
 	if( Wave->ResourceID )
 	{
@@ -129,7 +141,7 @@ FALSoundBuffer* FALSoundBuffer::CreateNativeBuffer( FALAudioDevice* AudioDevice,
 		// Keep track of associated resource name.
 		Buffer->ResourceName = Wave->GetPathName();
 
-		Buffer->InternalFormat = AudioDevice->GetInternalFormat( Wave->NumChannels );		
+		Buffer->InternalFormat = AudioDevice->GetInternalFormat( Wave->NumChannels );
 		Buffer->NumChannels = Wave->NumChannels;
 		Buffer->SampleRate = Wave->SampleRate;
 
diff --git a/Engine/Source/Runtime/HTML5/HTML5Audio/Private/HTML5AudioDevice.cpp b/Engine/Source/Runtime/HTML5/HTML5Audio/Private/HTML5AudioDevice.cpp
--- a/Engine/Source/Runtime/HTML5/HTML5Audio/Private/HTML5AudioDevice.cpp
+++ b/Engine/Source/Runtime/HTML5/HTML5Audio/Private/HTML5AudioDevice.cpp
@@ -15,6 +15,7 @@
 #include "AudioEffect.h"
 #include "Engine.h"
 
+DEFINE_LOG_CATEGORY(LogALAudio);
 
 //     2 UU == 1"
 // <=> 1 UU == 0.0127 m
@@ -40,7 +41,7 @@ IMPLEMENT_MODULE(FHTML5AudioDeviceModule, HTML5Audio );
 	UALAudioDevice constructor and UObject interface.
 ------------------------------------------------------------------------------------*/
 
-void FALAudioDevice::Teardown( void )
+void FALAudioDevice::TeardownHardware( void )
 {
 	// Flush stops all sources and deletes all buffers so sources can be safely deleted below.
 	Flush( NULL );
@@ -56,6 +57,8 @@ void FALAudioDevice::Teardown( void )
 	{
 		delete Sources[ i ];
 	}
+	Sources.Empty();
+	FreeSources.Empty();
 
 	// Disable the context
 	if( &alcMakeContextCurrent )
@@ -108,15 +111,13 @@ bool FALAudioDevice::InitializeHardware( void )
 	HardwareDevice = alcOpenDevice(  NULL );
 	if( !HardwareDevice )
 	{
-		//UE_LOG(LogHTML5Audio, Log, TEXT( "ALAudio: no OpenAL devices found." ) );
+		UE_LOG(LogALAudio, Log, TEXT( "ALAudio: no OpenAL devices found." ) );
 		return  false  ;
 	}
 
 	// Display the audio device that was actually opened
 	const ALCchar* OpenedDeviceName = alcGetString( HardwareDevice, ALC_DEVICE_SPECIFIER );
-	//UE_LOG(LogHTML5Audio, Log, TEXT( "ALAudio: no OpenAL devices found." ) );
-	//UE_LOG(LogHTML5Audio, Log, TEXT( "ALAudio device requested : %s" ), *DeviceName );
-	//UE_LOG(LogHTML5Audio, Log, TEXT( "ALAudio device Opened : %s" ), OpenedDeviceName );
+	UE_LOG(LogALAudio, Log, TEXT("ALAudio device opened : %s"), StringCast<TCHAR>(static_cast<const ANSICHAR*>(OpenedDeviceName)).Get());
 
 	// Create a context
 	int Caps[] = 
@@ -142,14 +143,14 @@ bool FALAudioDevice::InitializeHardware( void )
 	// Make sure everything happened correctly
 	if( alError( TEXT( "Init" ) ) )
 	{
-//		debugf( NAME_Init, TEXT( "ALAudio: alcMakeContextCurrent failed." ) );
+		UE_LOG(LogALAudio, Warning, TEXT("ALAudio: alcMakeContextCurrent failed."));
 		return false ;
 	}
 
-// 	debugf( NAME_Init, TEXT( "AL_VENDOR      : %s" ), ANSI_TO_TCHAR( ( ANSICHAR* )alGetString( AL_VENDOR ) ) );
-// 	debugf( NAME_Init, TEXT( "AL_RENDERER    : %s" ), ANSI_TO_TCHAR( ( ANSICHAR* )alGetString( AL_RENDERER ) ) );
-// 	debugf( NAME_Init, TEXT( "AL_VERSION     : %s" ), ANSI_TO_TCHAR( ( ANSICHAR* )alGetString( AL_VERSION ) ) );
-// 	debugf( NAME_Init, TEXT( "AL_EXTENSIONS  : %s" ), ANSI_TO_TCHAR( ( ANSICHAR* )alGetString( AL_EXTENSIONS ) ) );
+	UE_LOG(LogALAudio, Log, TEXT("AL_VENDOR      : %s"), StringCast<TCHAR>(static_cast<const ANSICHAR*>(alGetString(AL_VENDOR))).Get());
+	UE_LOG(LogALAudio, Log, TEXT("AL_RENDERER    : %s"), StringCast<TCHAR>(static_cast<const ANSICHAR*>(alGetString(AL_RENDERER))).Get());
+	UE_LOG(LogALAudio, Log, TEXT("AL_VERSION     : %s"), StringCast<TCHAR>(static_cast<const ANSICHAR*>(alGetString(AL_VERSION))).Get());
+	UE_LOG(LogALAudio, Log, TEXT("AL_EXTENSIONS  : %s"), StringCast<TCHAR>(static_cast<const ANSICHAR*>(alGetString(AL_EXTENSIONS))).Get());
  
 	// Get the enums for multichannel support
 #if !PLATFORM_HTML5
@@ -179,13 +180,13 @@ bool FALAudioDevice::InitializeHardware( void )
 
 	if( Sources.Num() < 1 )
 	{
-//		debugf( NAME_Error,TEXT( "ALAudio: couldn't allocate any sources" ) );
+		UE_LOG(LogALAudio, Warning, TEXT("ALAudio: couldn't allocate any sources"));
 		return false ;
 	}
 
 	// Update MaxChannels in case we couldn't create enough sources.
 	MaxChannels = Sources.Num();
-/*	debugf( NAME_Init, TEXT( "ALAudioDevice: Allocated %i sources" ), MaxChannels );*/
+	UE_LOG(LogALAudio, Verbose, TEXT("ALAudioDevice: Allocated %i sources"), MaxChannels);
 
 	// Use our own distance model.
 	alDistanceModel( AL_NONE );
@@ -196,7 +197,6 @@ bool FALAudioDevice::InitializeHardware( void )
 	// Initialized.
 	NextResourceID = 1;
 
-
 	return true ;
 }
 
@@ -364,12 +364,7 @@ OpenAL utility functions
 //
 bool FALAudioDevice::FindExt( const TCHAR* Name )
 {
-	if( alIsExtensionPresent( TCHAR_TO_ANSI( Name ) ) || alcIsExtensionPresent( HardwareDevice, TCHAR_TO_ANSI( Name ) ) )
-	{
-		//UE_LOG ( LogHTML5Audio, log,TEXT( "Device supports: %s" )); 
-		return( true );
-	}
-	return( true );
+	return alIsExtensionPresent(TCHAR_TO_ANSI(Name)) || alcIsExtensionPresent(HardwareDevice, TCHAR_TO_ANSI(Name));
 }
 
 //
@@ -403,22 +398,22 @@ bool FALAudioDevice::alError( const TCHAR* Text, bool Log )
 				switch ( Error )
 				{
 				case AL_INVALID_NAME:
-					UE_LOG ( LogInit, Log, TEXT( "ALAudio: AL_INVALID_NAME in %s" ), Text ); 
+						UE_LOG(LogALAudio, Warning, TEXT("ALAudio: AL_INVALID_NAME in %s"), Text);
 					break;
 				case AL_INVALID_VALUE:
-					UE_LOG ( LogInit, Log, TEXT( "ALAudio: AL_INVALID_VALUE in %s" ), Text ); 
+						UE_LOG(LogALAudio, Warning, TEXT("ALAudio: AL_INVALID_VALUE in %s"), Text);
 					break;
 				case AL_OUT_OF_MEMORY:
-					UE_LOG ( LogInit, Log, TEXT( "ALAudio: AL_OUT_OF_MEMORY in %s" ), Text ); 
+						UE_LOG(LogALAudio, Warning, TEXT("ALAudio: AL_OUT_OF_MEMORY in %s"), Text);
 					break;
 				case AL_INVALID_ENUM:
-					UE_LOG ( LogInit, Log, TEXT( "ALAudio: AL_INVALID_ENUM in %s" ), Text ); 
+						UE_LOG(LogALAudio, Warning, TEXT("ALAudio: AL_INVALID_ENUM in %s"), Text);
 					break;
 				case AL_INVALID_OPERATION:
-					UE_LOG ( LogInit, Log, TEXT( "ALAudio: AL_INVALID_OPERATION in %s" ), Text ); 
+						UE_LOG(LogALAudio, Warning, TEXT("ALAudio: AL_INVALID_OPERATION in %s"), Text);
 					break;
 				default:
-					UE_LOG ( LogInit, Log, TEXT( "ALAudio: Unknown Error NUM in %s" ), Text ); 
+						UE_LOG(LogALAudio, Warning, TEXT("ALAudio: Unknown Error NUM in %s"), Text);
 					break;
 				}
 			}
@@ -435,3 +430,21 @@ FSoundSource* FALAudioDevice::CreateSoundSource()
 {
 	return new FALSoundSource(this);
 }
+
+bool FALAudioDevice::HasCompressedAudioInfoClass(USoundWave* SoundWave)
+{
+#if WITH_OGGVORBIS
+	return true;
+#else
+	return false;
+#endif
+}
+
+class ICompressedAudioInfo* FALAudioDevice::CreateCompressedAudioInfo(USoundWave* SoundWave)
+{
+#if WITH_OGGVORBIS
+	return new FVorbisAudioInfo();
+#else
+	return NULL;
+#endif
+}
diff --git a/Engine/Source/Runtime/HTML5/HTML5Audio/Public/HTML5AudioDevice.h b/Engine/Source/Runtime/HTML5/HTML5Audio/Public/HTML5AudioDevice.h
--- a/Engine/Source/Runtime/HTML5/HTML5Audio/Public/HTML5AudioDevice.h
+++ b/Engine/Source/Runtime/HTML5/HTML5Audio/Public/HTML5AudioDevice.h
@@ -27,6 +27,8 @@ class FALAudioDevice;
 #pragma pack (pop)
 #endif 
 
+DECLARE_LOG_CATEGORY_EXTERN(LogALAudio, Log, All);
+
 /**
  * OpenAL implementation of FSoundBuffer, containing the wave data and format information.
  */
@@ -204,6 +206,12 @@ public:
 	/** Starts up any platform specific hardware/APIs */
 	virtual bool InitializeHardware() OVERRIDE;
  
+	/**
+	 * Tears down audio device by stopping all sounds, removing all buffers,  destroying all sources, ... Called by both Destroy and ShutdownAfterError
+	 * to perform the actual tear down.
+	 */
+	virtual void TeardownHardware() OVERRIDE;
+
 	/**
 	 * Update the audio device and calculates the cached inverse transform later
 	 * on used for spatialization.
@@ -224,10 +232,14 @@ public:
 	 */
 	virtual void FreeResource( USoundWave* SoundWave );
 
+	virtual bool HasCompressedAudioInfoClass(USoundWave* SoundWave) OVERRIDE;
+
+	virtual class ICompressedAudioInfo* CreateCompressedAudioInfo(USoundWave* SoundWave) OVERRIDE;
+
 	void FindProcs( bool AllowExt );
 
 	// Error checking.
-	bool alError( const TCHAR* Text, bool Log = 1 );
+	bool alError( const TCHAR* Text, bool Log = true );
 
 protected:
 
@@ -239,11 +251,6 @@ protected:
 	void FindProc( void*& ProcAddress, char* Name, char* SupportName, bool& Supports, bool AllowExt );
 	bool FindExt( const TCHAR* Name );
 
-	/**
-	 * Tears down audio device by stopping all sounds, removing all buffers,  destroying all sources, ... Called by both Destroy and ShutdownAfterError
-	 * to perform the actual tear down.
-	 */
-	void Teardown( void );
 
 	// Variables.
 
