#UE4 #HTML5

HTML5 browser build fixes for NFS over http.
Todo: Correctly generate html5 file when serving via NFS.
howto:
- flip relevant bools in UEBuildHTML5.cs
- build html5/cooker.
- run cook on the fly server via UFE.
- Remove .data/.data.js references from .html ( see todo)
- point browser at http://cookserver:DEFAULT_FILE_SERVING_PORT/gamename.html

diff --git a/Engine/Source/Runtime/HTML5/HTML5JS/Private/HTML5JavaScriptFx.js b/Engine/Source/Runtime/HTML5/HTML5JS/Private/HTML5JavaScriptFx.js
--- a/Engine/Source/Runtime/HTML5/HTML5JS/Private/HTML5JavaScriptFx.js
+++ b/Engine/Source/Runtime/HTML5/HTML5JS/Private/HTML5JavaScriptFx.js
@@ -8,16 +8,15 @@ var UE_JavaScriptLibary =
 
 		var request = new XMLHttpRequest();
 		if (insize && indata) {
-			var postData = Module.HEAP8.subarray(indata, indata + insize);
-			request.open('POST', _url, false);
+		    var postData = Module.HEAP8.subarray(indata, indata + insize);
+		    request.open('POST', _url, false);
+		    request.overrideMimeType('text\/plain; charset=x-user-defined'); 
 			request.send(postData);
 		} else {
 			request.open('GET', _url, false);
 			request.send();
 		}
 
-		console.log("Fetching " + _url +  request.lenght);
-
 		if (request.status != 200) {
 			console.log("Fetching " + _url + " failed: " + request.responseText);
 
@@ -28,7 +27,7 @@ var UE_JavaScriptLibary =
 		}
 
 		// we got the XHR result as a string.  We need to write this to Module.HEAP8[outdataptr]
-		var replyString = request.response;
+		var replyString = request.responseText;
 		var replyLength = replyString.length;
 
 		var outdata = Module._malloc(replyLength);
@@ -44,7 +43,7 @@ var UE_JavaScriptLibary =
 		// tears and crying.  Copy from the result-string into the heap.
 		var replyDest = Module.HEAP8.subarray(outdata, outdata + replyLength);
 		for (var i = 0; i < replyLength; ++i) {
-			replyDest[i] = replyString.charCodeAt(i);
+			replyDest[i] = replyString.charCodeAt(i) &  0xff;
 		}
 
 		Module.HEAP32[outsizeptr >> 2] = replyLength;
diff --git a/Engine/Source/Runtime/NetworkFile/Private/HTTPTransport.cpp b/Engine/Source/Runtime/NetworkFile/Private/HTTPTransport.cpp
--- a/Engine/Source/Runtime/NetworkFile/Private/HTTPTransport.cpp
+++ b/Engine/Source/Runtime/NetworkFile/Private/HTTPTransport.cpp
@@ -9,6 +9,10 @@
 #include "WinHttp.h"
 #endif 
 
+#if PLATFORM_HTML5_BROWSER
+#include "HTML5JavaScriptFx.h"
+#endif 
+
 
 FHTTPTransport::FHTTPTransport()
 	:Guid(FGuid::NewGuid())
@@ -17,7 +21,6 @@ FHTTPTransport::FHTTPTransport()
 
 bool FHTTPTransport::Initialize(const TCHAR* HostIp)
 {
-	TCHAR Url[1048];
 	FCString::Sprintf( Url, TEXT("http://%s:%d"), HostIp, (int)DEFAULT_FILE_SERVING_PORT); 
 
 #if !PLATFORM_HTML5
@@ -119,11 +122,13 @@ bool FHTTPTransport::SendPayloadAndReceiveResponse(TArray<uint8>& In, TArray<uin
 #endif 
 
 #if PLATFORM_HTML5_BROWSER
-	UE_SendAndRecievePayLoad("http://127.0.0.1:9000",(char*)Ar.GetData(),Ar.Num(),(char**)&OutData,(int*)&OutSize);
-	Out.Append(OutData,OutSize); 
-	// don't go through the Unreal Memory system. 
-	::free (OutData); 
-
+	UE_SendAndRecievePayLoad(TCHAR_TO_ANSI(Url),(char*)Ar.GetData(),Ar.Num(),(char**)&OutData,(int*)&OutSize);
+	if (OutSize)
+	{
+		Out.Append(OutData,OutSize); 
+		// don't go through the Unreal Memory system. 
+		::free (OutData); 
+	}
 	return true; 
 #endif 
 #endif
diff --git a/Engine/Source/Runtime/NetworkFile/Private/HTTPTransport.h b/Engine/Source/Runtime/NetworkFile/Private/HTTPTransport.h
--- a/Engine/Source/Runtime/NetworkFile/Private/HTTPTransport.h
+++ b/Engine/Source/Runtime/NetworkFile/Private/HTTPTransport.h
@@ -2,7 +2,10 @@
 #pragma once
 #if ENABLE_HTTP_FOR_NFS
 #include "ITransport.h"
+
+#if !PLATFORM_HTML5
 #include "Http.h"
+#endif 
 
 class FHTTPTransport : public ITransport
 {
@@ -22,6 +25,7 @@ private:
 #endif 
 
 	FGuid Guid; 
+	TCHAR Url[1048];
 
 };
 #endif
\ No newline at end of file
