Removing more duplicated variables from audio sub classes

Duplicate ResourceID would have caused the same problems in HTML5 and Linux as Android. Other duplications are errors waiting to happen, there are a couple of stats functions that could be using the base class versions and not displaying the correct information.

diff --git a/Engine/Source/Runtime/Android/AndroidAudio/Public/AndroidAudioDevice.h b/Engine/Source/Runtime/Android/AndroidAudio/Public/AndroidAudioDevice.h
--- a/Engine/Source/Runtime/Android/AndroidAudio/Public/AndroidAudioDevice.h
+++ b/Engine/Source/Runtime/Android/AndroidAudio/Public/AndroidAudioDevice.h
@@ -97,26 +97,13 @@ public:
 		return( BufferSize ); 
 	}
 
-	/** 
-	 * Returns the number of channels for this buffer
-	 */
-	int GetNumChannels( void ) 
-	{ 
-		return( NumChannels ); 
-	}
-	
-		
 	/** Audio device this buffer is attached to */
 	FSLESAudioDevice*			AudioDevice;
 	/** Data */
 	uint8*					AudioData;
-	/** Human readable name of resource, most likely name of UObject associated during caching. */
-	FString					ResourceName;
 
 	/** Number of bytes stored in OpenAL, or the size of the ogg vorbis data */
 	int						BufferSize;
-	/** The number of channels in this sound buffer - should be directly related to InternalFormat */
-	int						NumChannels;
 	/** Sample rate of the ogg vorbis data - typically 44100 or 22050 */
 	int						SampleRate;
 
diff --git a/Engine/Source/Runtime/HTML5/HTML5Audio/Public/HTML5AudioDevice.h b/Engine/Source/Runtime/HTML5/HTML5Audio/Public/HTML5AudioDevice.h
--- a/Engine/Source/Runtime/HTML5/HTML5Audio/Public/HTML5AudioDevice.h
+++ b/Engine/Source/Runtime/HTML5/HTML5Audio/Public/HTML5AudioDevice.h
@@ -78,29 +78,15 @@ public:
 		return( BufferSize ); 
 	}
 
-	/** 
-	 * Returns the number of channels for this buffer
-	 */
-	int GetNumChannels( void ) 
-	{ 
-		return( NumChannels ); 
-	}
-		
 	/** Audio device this buffer is attached to */
 	FALAudioDevice*			AudioDevice;
 	/** Array of buffer ids used to reference the data stored in AL. */
 	ALuint					BufferIds[2];
-	/** Resource ID of associated USoundeWave */
-	int						ResourceID;
-	/** Human readable name of resource, most likely name of UObject associated during caching. */
-	FString					ResourceName;
 	/** Format of the data internal to OpenAL */
 	ALuint					InternalFormat;
 
 	/** Number of bytes stored in OpenAL, or the size of the ogg vorbis data */
 	int						BufferSize;
-	/** The number of channels in this sound buffer - should be directly related to InternalFormat */
-	int						NumChannels;
 	/** Sample rate of the ogg vorbis data - typically 44100 or 22050 */
 	int						SampleRate;
 };
diff --git a/Engine/Source/Runtime/Linux/ALAudio/Public/ALAudioDevice.h b/Engine/Source/Runtime/Linux/ALAudio/Public/ALAudioDevice.h
--- a/Engine/Source/Runtime/Linux/ALAudio/Public/ALAudioDevice.h
+++ b/Engine/Source/Runtime/Linux/ALAudio/Public/ALAudioDevice.h
@@ -76,29 +76,15 @@ public:
 		return( BufferSize ); 
 	}
 
-	/** 
-	 * Returns the number of channels for this buffer
-	 */
-	int GetNumChannels( void ) 
-	{ 
-		return( NumChannels ); 
-	}
-		
 	/** Audio device this buffer is attached to */
 	FALAudioDevice*			AudioDevice;
 	/** Array of buffer ids used to reference the data stored in AL. */
 	ALuint					BufferIds[2];
-	/** Resource ID of associated USoundeWave */
-	int						ResourceID;
-	/** Human readable name of resource, most likely name of UObject associated during caching. */
-	FString					ResourceName;
 	/** Format of the data internal to OpenAL */
 	ALuint					InternalFormat;
 
 	/** Number of bytes stored in OpenAL, or the size of the ogg vorbis data */
 	int						BufferSize;
-	/** The number of channels in this sound buffer - should be directly related to InternalFormat */
-	int						NumChannels;
 	/** Sample rate of the ogg vorbis data - typically 44100 or 22050 */
 	int						SampleRate;
 };
diff --git a/Engine/Source/Runtime/Mac/CoreAudio/Private/CoreAudioBuffer.cpp b/Engine/Source/Runtime/Mac/CoreAudio/Private/CoreAudioBuffer.cpp
--- a/Engine/Source/Runtime/Mac/CoreAudio/Private/CoreAudioBuffer.cpp
+++ b/Engine/Source/Runtime/Mac/CoreAudio/Private/CoreAudioBuffer.cpp
@@ -30,9 +30,6 @@ FCoreAudioSoundBuffer::FCoreAudioSoundBuffer( FAudioDevice* InAudioDevice, ESoun
 	PCMData( NULL ),
 	PCMDataSize( 0 ),
 	DecompressionState( NULL ),
-	NumChannels( 0 ),
-	ResourceID( 0 ),
-	bAllocationInPermanentPool( false ),
 	bDynamicResource( false )
 {
 }
diff --git a/Engine/Source/Runtime/Mac/CoreAudio/Public/CoreAudioDevice.h b/Engine/Source/Runtime/Mac/CoreAudio/Public/CoreAudioDevice.h
--- a/Engine/Source/Runtime/Mac/CoreAudio/Public/CoreAudioDevice.h
+++ b/Engine/Source/Runtime/Mac/CoreAudio/Public/CoreAudioDevice.h
@@ -157,14 +157,6 @@ public:
 	
 	/** Wrapper to handle the decompression of audio codecs */
 	class ICompressedAudioInfo*		DecompressionState;
-	/** Cumulative channels from all streams */
-	int32							NumChannels;
-	/** Resource ID of associated USoundWave */
-	int32							ResourceID;
-	/** Human readable name of resource, most likely name of UObject associated during caching.	*/
-	FString						ResourceName;
-	/** Whether memory for this buffer has been allocated from permanent pool. */
-	bool						bAllocationInPermanentPool;
 	/** Set to true when the PCM data should be freed when the buffer is destroyed */
 	bool						bDynamicResource;
 };
