#UE4 #HTML5 Fix up Win32 Simulator Build for HTML5.

diff --git a/Engine/Source/Runtime/Core/Public/HTML5/HTML5PlatformMisc.h b/Engine/Source/Runtime/Core/Public/HTML5/HTML5PlatformMisc.h
--- a/Engine/Source/Runtime/Core/Public/HTML5/HTML5PlatformMisc.h
+++ b/Engine/Source/Runtime/Core/Public/HTML5/HTML5PlatformMisc.h
@@ -42,7 +42,7 @@ struct CORE_API FHTML5Misc : public FGenericPlatformMisc
 		if (IsDebuggerPresent())
 		{
 #if PLATFORM_HTML5_WIN32
-			exit(-1);
+			__debugbreak();
 #else
 			emscripten_log(255, "DebugBreak() called!");
 #endif
diff --git a/Engine/Source/Runtime/HTML5/Simulator/HTML5Win32/Private/HTML5Win32PrivatePCH.h b/Engine/Source/Runtime/HTML5/Simulator/HTML5Win32/Private/HTML5Win32PrivatePCH.h
new file mode 100644
--- /dev/null
+++ b/Engine/Source/Runtime/HTML5/Simulator/HTML5Win32/Private/HTML5Win32PrivatePCH.h
@@ -0,0 +1,5 @@
+// Copyright 1998-2014 Epic Games, Inc. All Rights Reserved.
+
+// This file is intentionally empty to keep the UHT happy, HTML5Win32 doesn't depend on internal UE headers. 
+
+
diff --git a/Engine/Source/Runtime/HTML5/Simulator/HTML5Win32/Private/IPAddressRaw.cpp b/Engine/Source/Runtime/HTML5/Simulator/HTML5Win32/Private/IPAddressRaw.cpp
--- a/Engine/Source/Runtime/HTML5/Simulator/HTML5Win32/Private/IPAddressRaw.cpp
+++ b/Engine/Source/Runtime/HTML5/Simulator/HTML5Win32/Private/IPAddressRaw.cpp
@@ -1,5 +1,6 @@
 // Copyright 1998-2014 Epic Games, Inc. All Rights Reserved.
 
+#include "HTML5Win32PrivatePCH.h"
 #include "IPAddressRaw.h"
 
 
diff --git a/Engine/Source/Runtime/HTML5/Simulator/HTML5Win32/Private/LoadDLL.cpp b/Engine/Source/Runtime/HTML5/Simulator/HTML5Win32/Private/LoadDLL.cpp
--- a/Engine/Source/Runtime/HTML5/Simulator/HTML5Win32/Private/LoadDLL.cpp
+++ b/Engine/Source/Runtime/HTML5/Simulator/HTML5Win32/Private/LoadDLL.cpp
@@ -1,5 +1,7 @@
 // Copyright 1998-2014 Epic Games, Inc. All Rights Reserved.
 
+
+#include "HTML5Win32PrivatePCH.h"
 #include "LoadDLL.h" 
 #include <windows.h>
 #include <string>
diff --git a/Engine/Source/Runtime/HTML5/Simulator/HTML5Win32/Private/SocketRaw.cpp b/Engine/Source/Runtime/HTML5/Simulator/HTML5Win32/Private/SocketRaw.cpp
--- a/Engine/Source/Runtime/HTML5/Simulator/HTML5Win32/Private/SocketRaw.cpp
+++ b/Engine/Source/Runtime/HTML5/Simulator/HTML5Win32/Private/SocketRaw.cpp
@@ -1,5 +1,6 @@
 // Copyright 1998-2014 Epic Games, Inc. All Rights Reserved.
 
+#include "HTML5Win32PrivatePCH.h"
 #include "SocketRaw.h"
 #include "IPAddressRaw.h"
 #pragma warning( disable : 4530 )
diff --git a/Engine/Source/Runtime/HTML5/Simulator/HTML5Win32/Private/WinHttp.cpp b/Engine/Source/Runtime/HTML5/Simulator/HTML5Win32/Private/WinHttp.cpp
--- a/Engine/Source/Runtime/HTML5/Simulator/HTML5Win32/Private/WinHttp.cpp
+++ b/Engine/Source/Runtime/HTML5/Simulator/HTML5Win32/Private/WinHttp.cpp
@@ -1,6 +1,6 @@
 // Copyright 1998-2014 Epic Games, Inc. All Rights Reserved.
 
-
+#include "HTML5Win32PrivatePCH.h"
 #include "WinHttp.h"
 #include "curl/curl.h"
 
@@ -36,7 +36,7 @@ static size_t
 static char GURL[1024];
 void HTML5Win32::NFSHttp::Init(char* URL)
 {
-	 strcpy(GURL,URL); 
+	 strcpy_s(GURL,URL); 
 	 curl_global_init(CURL_GLOBAL_ALL); 
 	 curl = curl_easy_init();
 }
diff --git a/Engine/Source/Runtime/Sockets/Private/HTML5/Simulator/IPAddressX.cpp b/Engine/Source/Runtime/Sockets/Private/HTML5/Simulator/IPAddressX.cpp
--- a/Engine/Source/Runtime/Sockets/Private/HTML5/Simulator/IPAddressX.cpp
+++ b/Engine/Source/Runtime/Sockets/Private/HTML5/Simulator/IPAddressX.cpp
@@ -1,5 +1,6 @@
 // Copyright 1998-2014 Epic Games, Inc. All Rights Reserved.
 
+#include "SocketsPrivatePCH.h"
 #include "IPAddressX.h"
 #include "Sockets/IPAddressRaw.h" 
 
diff --git a/Engine/Source/Runtime/Sockets/Private/HTML5/Simulator/SocketX.cpp b/Engine/Source/Runtime/Sockets/Private/HTML5/Simulator/SocketX.cpp
--- a/Engine/Source/Runtime/Sockets/Private/HTML5/Simulator/SocketX.cpp
+++ b/Engine/Source/Runtime/Sockets/Private/HTML5/Simulator/SocketX.cpp
@@ -1,5 +1,6 @@
 // Copyright 1998-2014 Epic Games, Inc. All Rights Reserved.
 
+#include "SocketsPrivatePCH.h"
 #include "SocketX.h"
 #include "Sockets/SocketRaw.h"
 #include "Sockets/IPAddressRaw.h"
