Fix HTML5 load game logic when no saved game exists, needed to make ShooterGame work.

diff --git a/Engine/Source/Runtime/Engine/Public/SaveGameSystem.h b/Engine/Source/Runtime/Engine/Public/SaveGameSystem.h
--- a/Engine/Source/Runtime/Engine/Public/SaveGameSystem.h
+++ b/Engine/Source/Runtime/Engine/Public/SaveGameSystem.h
@@ -77,7 +77,9 @@ public:
 #if PLATFORM_HTML5_BROWSER
 		char*	OutData;
 		int		Size;
-		UE_LoadGame(TCHAR_TO_ANSI(Name),UserIndex,&OutData,&Size);
+		bool Result = UE_LoadGame(TCHAR_TO_ANSI(Name),UserIndex,&OutData,&Size);
+		if (!Result)
+			return false; 
 		Data.Append((uint8*)OutData,Size);
 		::free (OutData);
 		return true;
diff --git a/Engine/Source/Runtime/HTML5/HTML5JS/Private/HTML5JavaScriptFx.js b/Engine/Source/Runtime/HTML5/HTML5JS/Private/HTML5JavaScriptFx.js
--- a/Engine/Source/Runtime/HTML5/HTML5JS/Private/HTML5JavaScriptFx.js
+++ b/Engine/Source/Runtime/HTML5/HTML5JS/Private/HTML5JavaScriptFx.js
@@ -66,8 +66,8 @@ var UE_JavaScriptLibary =
 		var _name = Pointer_stringify(name);
 		// local storage only takes strings, we need to convert string to base64 before storing.
 		var b64encoded = $.jStorage.get(_name);
-		if (typeof b64encoded == null)
-			return false;
+		if (typeof b64encoded == null || b64encoded == null)
+		    return false;
 		var decodedArray = base64DecToArr(b64encoded);
 		// copy back the decoded array.
 		var outdata = Module._malloc(decodedArray.length);
