CIS build fix for HTML5

-initialization list typo

diff --git a/Engine/Source/Runtime/HTML5/HTML5Audio/Private/HTML5AudioDevice.cpp b/Engine/Source/Runtime/HTML5/HTML5Audio/Private/HTML5AudioDevice.cpp
--- a/Engine/Source/Runtime/HTML5/HTML5Audio/Private/HTML5AudioDevice.cpp
+++ b/Engine/Source/Runtime/HTML5/HTML5Audio/Private/HTML5AudioDevice.cpp
@@ -33,7 +33,7 @@ IMPLEMENT_MODULE(FHTML5AudioDeviceModule, HTML5Audio );
 	UALAudioDevice constructor and UObject interface.
 ------------------------------------------------------------------------------------*/
 FALAudioDevice::FALAudioDevice()
-	, HardwareDevice(nullptr)
+	: HardwareDevice(nullptr)
 	, SoundContext(nullptr)
 {}
 
