HTML5: Correctly gate map downloading. Use a progress bar with real progress instead of an infinite throbber.

diff --git a/Engine/Source/Programs/AutomationTool/HTML5/HTML5Platform.Automation.cs b/Engine/Source/Programs/AutomationTool/HTML5/HTML5Platform.Automation.cs
--- a/Engine/Source/Programs/AutomationTool/HTML5/HTML5Platform.Automation.cs
+++ b/Engine/Source/Programs/AutomationTool/HTML5/HTML5Platform.Automation.cs
@@ -370,7 +370,7 @@ public class HTML5Platform : Platform
     public override List<string> GetCookExtraMaps()
     {
         var Maps = new List<string>();
-        Maps.Add("/Engine/Maps/Loading");
+        Maps.Add("/Engine/Maps/Entry");
         return Maps; 
     }
 
diff --git a/Engine/Source/Runtime/HTML5/MapPakDownloader/MapPakDownloader.Build.cs b/Engine/Source/Runtime/HTML5/MapPakDownloader/MapPakDownloader.Build.cs
--- a/Engine/Source/Runtime/HTML5/MapPakDownloader/MapPakDownloader.Build.cs
+++ b/Engine/Source/Runtime/HTML5/MapPakDownloader/MapPakDownloader.Build.cs
@@ -12,6 +12,7 @@ namespace UnrealBuildTool.Rules
 					"Core",
                     "CoreUObject",
                     "Engine",
+                    "Slate",
 					// ... add other public dependencies that you statically link with here ...
 				}
 				);
diff --git a/Engine/Source/Runtime/HTML5/MapPakDownloader/Private/MapPakDownloader.cpp b/Engine/Source/Runtime/HTML5/MapPakDownloader/Private/MapPakDownloader.cpp
--- a/Engine/Source/Runtime/HTML5/MapPakDownloader/Private/MapPakDownloader.cpp
+++ b/Engine/Source/Runtime/HTML5/MapPakDownloader/Private/MapPakDownloader.cpp
@@ -5,6 +5,7 @@
 #include "MapPakDownloader.h"
 #include "emscripten.h"
 #include "Misc/Guid.h"
+#include "SlateExtras.h"
 
 DECLARE_DELEGATE_OneParam(FDelegateFString, FString)
 DECLARE_DELEGATE_OneParam(FDelegateInt32, int32)
@@ -72,7 +73,27 @@ public:
 
 };
 
+class FloatOption : public TSharedFromThis<FloatOption>
+{
+public: 
+	FloatOption() : 
+		Value(0.0f) 
+	{} 
+	void SetFloat(float InFloat)
+	{ 
+		Value = InFloat; 
+	}
+	TOptional<float> GetFloat() const
+	{ 
+		return Value; 
+	}
+private: 
+	float Value;
+};
+
+
 FMapPakDownloader::FMapPakDownloader()
+	:IsTransitionLevel(false)
 {}
 
 bool FMapPakDownloader::Init()
@@ -95,6 +116,38 @@ bool FMapPakDownloader::Init()
 	IPlatformFile& PlatformFile = FPlatformFileManager::Get().GetPlatformFile();
 	PlatformFile.CreateDirectory(*PakLocation);
 
+	ProgressContainter = MakeShareable(new FloatOption());
+
+    // Thin progress bar. Change this widget if you want a custom loading screen. 
+	LoadingWidget = SNew(SBox)
+		.VAlign(VAlign_Center)
+		.HAlign(HAlign_Center)
+		[
+			SNew(SBox)
+			.VAlign(VAlign_Fill)
+			.HAlign(HAlign_Fill)
+			.HeightOverride(FOptionalSize(4))
+			.WidthOverride(FOptionalSize(300))
+			[
+				SNew(SProgressBar)
+				.Percent(ProgressContainter->AsShared(), &FloatOption::GetFloat)
+				.BorderPadding(FVector2D(0, 0))	
+				.BarFillType(EProgressBarFillType::LeftToRight)
+			]
+		];
+
+	// Hook up PostLoad 
+	FSimpleDelegate PostLoadHandler = FSimpleDelegate::CreateLambda([=](){
+		{
+			if (IsTransitionLevel)
+			{
+				GEngine->GameViewport->AddViewportWidgetContent(LoadingWidget->AsShared());
+			}
+		}
+	});
+
+	FCoreUObjectDelegates::PostLoadMap.Add(PostLoadHandler);
+
 	return true;
 }
 
@@ -118,16 +171,21 @@ void FMapPakDownloader::CachePak()
 											// Get hold of the world.
 											TObjectIterator<UWorld> It;
 											UWorld* ItWorld = *It;
+											if (IsTransitionLevel)
+												GEngine->GameViewport->RemoveViewportWidgetContent(LoadingWidget->AsShared());
 
 											UE_LOG(LogMapPakDownloader, Warning, TEXT("Travel to %s"), *MapToCache);
 											// Make engine Travel to the cached Map. 
 											GEngine->SetClientTravel(ItWorld, *MapToCache, TRAVEL_Absolute);
 
+											IsTransitionLevel = false; 
+											ProgressContainter->SetFloat(0.0f); 
 											// delete the HTTP request. 
 											delete PakRequest; 
 										});
 
 	FDelegateInt32 OnFileDownloadProgress = FDelegateInt32::CreateLambda([=](int32 Progress){
+												ProgressContainter->SetFloat((float)Progress/100.0f);
 												UE_LOG(LogMapPakDownloader, Warning, TEXT(" %s %d%% downloaded"), *PakRequest->GetFileName(), Progress);
 											});
 
@@ -156,12 +214,28 @@ void FMapPakDownloader::CachePak()
 
 void FMapPakDownloader::Cache(FString& Map, FString& InLastMap, void* InDynData)
 {
-	MapToCache = Map;
-	LastMap = InLastMap;
+	bool UseMapDownloader = false;
+	if (GConfig)
+	{
+		GConfig->GetBool(TEXT("/Script/HTML5PlatformEditor.HTML5TargetSettings"), TEXT("UseAsyncLevelLoading"), UseMapDownloader, GEngineIni);
+	}
 
-	if (!FPackageName::DoesPackageExist(Map))
+	if (UseMapDownloader)
 	{
-		CachePak();
-		Map = TEXT("/Engine/Maps/Loading");
+		MapToCache = Map;
+		LastMap = InLastMap;
+
+		FString OutLongPackageName;
+		FString OutFileName;
+
+		if (!FPackageName::SearchForPackageOnDisk(Map, &OutLongPackageName, &OutFileName, false))
+		{
+			UE_LOG(LogMapPakDownloader, Warning, TEXT("Caching.... %s"), *Map);
+			CachePak();
+			Map = TEXT("/Engine/Maps/Entry");
+			IsTransitionLevel = true;
+		}
+
 	}
 }
+
diff --git a/Engine/Source/Runtime/HTML5/MapPakDownloader/Public/MapPakDownloader.h b/Engine/Source/Runtime/HTML5/MapPakDownloader/Public/MapPakDownloader.h
--- a/Engine/Source/Runtime/HTML5/MapPakDownloader/Public/MapPakDownloader.h
+++ b/Engine/Source/Runtime/HTML5/MapPakDownloader/Public/MapPakDownloader.h
@@ -3,6 +3,8 @@
 #pragma once
 
 #include "Core.h"
+#include "SlateCore.h"
+#include "SlateFwd.h"
 
 class FMapPakDownloader
 {
@@ -43,5 +45,13 @@ private:
 	// Relative dir of Pak files 
 	FString PakLocation;
 
+	// In Loading screen level.
+	bool IsTransitionLevel; 
+	
+	// Loading Screen Widget.
+	TSharedPtr<SWidget> LoadingWidget;
+
+	//  Progress
+	TSharedPtr<class FloatOption>  ProgressContainter;
 };
 
