PR #1020: Normalized JS style. No functional change. (Contributed by [REDACTED_PERSON])

diff --git a/Engine/Source/Runtime/HTML5/HTML5JS/Private/HTML5JavaScriptFx.js b/Engine/Source/Runtime/HTML5/HTML5JS/Private/HTML5JavaScriptFx.js
--- a/Engine/Source/Runtime/HTML5/HTML5JS/Private/HTML5JavaScriptFx.js
+++ b/Engine/Source/Runtime/HTML5/HTML5JS/Private/HTML5JavaScriptFx.js
@@ -1,173 +1,171 @@
 // Copyright 1998-2015 Epic Games, Inc. All Rights Reserved.
 
+var UE_JavaScriptLibary = {
+  UE_SendAndRecievePayLoad: function (url, indata, insize, outdataptr, outsizeptr) {
+    var _url = Pointer_stringify(url);
+
+    var request = new XMLHttpRequest();
+    if (insize && indata) {
+      var postData = Module.HEAP8.subarray(indata, indata + insize);
+      request.open('POST', _url, false);
+      request.overrideMimeType('text\/plain; charset=x-user-defined');
+      request.send(postData);
+    } else {
+      request.open('GET', _url, false);
+      request.send();
+    }
+
+    if (request.status != 200) {
+      console.log("Fetching " + _url + " failed: " + request.responseText);
+
+      Module.HEAP32[outsizeptr >> 2] = 0;
+      Module.HEAP32[outdataptr >> 2] = 0;
+
+      return;
+    }
+
+    // we got the XHR result as a string.  We need to write this to Module.HEAP8[outdataptr]
+    var replyString = request.responseText;
+    var replyLength = replyString.length;
+
+    var outdata = Module._malloc(replyLength);
+    if (!outdata) {
+      console.log("Failed to allocate " + replyLength + " bytes in heap for reply");
+
+      Module.HEAP32[outsizeptr >> 2] = 0;
+      Module.HEAP32[outdataptr >> 2] = 0;
+
+      return;
+    }
+
+    // tears and crying.  Copy from the result-string into the heap.
+    var replyDest = Module.HEAP8.subarray(outdata, outdata + replyLength);
+    for (var i = 0; i < replyLength; ++i) {
+      replyDest[i] = replyString.charCodeAt(i) &  0xff;
+    }
+
+    Module.HEAP32[outsizeptr >> 2] = replyLength;
+    Module.HEAP32[outdataptr >> 2] = outdata;
+  },
+
+  UE_SaveGame: function (name, userIndex, indata, insize){
+    // user index is not used.
+    var _name = Pointer_stringify(name);
+    var gamedata = Module.HEAP8.subarray(indata, indata + insize);
+    // local storage only takes strings, we need to convert string to base64 before storing.
+    var b64encoded = base64EncArr(gamedata);
+    $.jStorage.set(_name, b64encoded);
+    return true;
+  },
+
+  UE_LoadGame: function (name, userIndex, outdataptr, outsizeptr){
+    var _name = Pointer_stringify(name);
+    // local storage only takes strings, we need to convert string to base64 before storing.
+    var b64encoded = $.jStorage.get(_name);
+    if (b64encoded === null)
+      return false;
+    var decodedArray = base64DecToArr(b64encoded);
+    // copy back the decoded array.
+    var outdata = Module._malloc(decodedArray.length);
+    // view the allocated data as a HEAP8.
+    var dest = Module.HEAP8.subarray(outdata, outdata + decodedArray.length);
+    // copy back.
+    for (var i = 0; i < decodedArray.length; ++i) {
+      dest[i] = decodedArray[i];
+    }
+    Module.HEAP32[outsizeptr >> 2] = decodedArray.length;
+    Module.HEAP32[outdataptr >> 2] = outdata;
+    return true;
+  },
+
+  UE_DoesSaveGameExist: function (name, userIndex){
+    var _name = Pointer_stringify(name);
+    var b64encoded = $.jStorage.get(_name);
+    if (b64encoded === null)
+      return false;
+    return true;
+  },
+
+  UE_MessageBox: function (type, message, caption ) {
+    // type maps to EAppMsgType::Type
+    var text;
+    if ( type === 0 ) {
+      text = Pointer_stringify(message);
+      if (!confirm(text))
+        return 0;
+    } else {
+      text = Pointer_stringify(message);
+      alert(text);
+    }
+    return 1;
+  },
+
+  UE_GetCurrentCultureName: function (address, outsize) {
+    var culture_name = navigator.language || navigator.browserLanguage;
+    if (culture_name.lenght >= outsize)
+    	return 0;
+    Module.writeAsciiToMemory(culture_name, address);
+    return 1;
+  },
+
+  UE_MakeHTTPDataRequest: function (ctx, url, verb, payload, payloadsize, headers, freeBuffer, onload, onerror, onprogress) {
+    var _url = Pointer_stringify(url);
+    var _verb = Pointer_stringify(verb);
+    var _headers = Pointer_stringify(headers);
+
+    var xhr = new XMLHttpRequest();
+    xhr.open(_verb, _url, true);
+    xhr.responseType = 'arraybuffer';
+
+    // set all headers. 
+    var _headerArray = _headers.split("%");
+    for(var headerArrayidx = 0; headerArrayidx < _headerArray.length; headerArrayidx++){
+      var header = _headerArray[headerArrayidx].split(":");
+      xhr.setRequestHeader(header[0],header[1]);
+    }
 
-var UE_JavaScriptLibary =
-{
-	UE_SendAndRecievePayLoad: function (url, indata, insize, outdataptr, outsizeptr) {
-		var _url = Pointer_stringify(url)
-
-		var request = new XMLHttpRequest();
-		if (insize && indata) {
-		    var postData = Module.HEAP8.subarray(indata, indata + insize);
-		    request.open('POST', _url, false);
-		    request.overrideMimeType('text\/plain; charset=x-user-defined'); 
-			request.send(postData);
-		} else {
-			request.open('GET', _url, false);
-			request.send();
-		}
-
-		if (request.status != 200) {
-			console.log("Fetching " + _url + " failed: " + request.responseText);
-
-			Module.HEAP32[outsizeptr >> 2] = 0;
-			Module.HEAP32[outdataptr >> 2] = 0;
-
-			return;
-		}
-
-		// we got the XHR result as a string.  We need to write this to Module.HEAP8[outdataptr]
-		var replyString = request.responseText;
-		var replyLength = replyString.length;
-
-		var outdata = Module._malloc(replyLength);
-		if (!outdata) {
-			console.log("Failed to allocate " + replyLength + " bytes in heap for reply");
-
-			Module.HEAP32[outsizeptr >> 2] = 0;
-			Module.HEAP32[outdataptr >> 2] = 0;
-
-			return;
-		}
-
-		// tears and crying.  Copy from the result-string into the heap.
-		var replyDest = Module.HEAP8.subarray(outdata, outdata + replyLength);
-		for (var i = 0; i < replyLength; ++i) {
-			replyDest[i] = replyString.charCodeAt(i) &  0xff;
-		}
-
-		Module.HEAP32[outsizeptr >> 2] = replyLength;
-		Module.HEAP32[outdataptr >> 2] = outdata;
-	},
-
-
-	UE_SaveGame: function (name, userIndex, indata, insize){
-			// user index is not used.
-			var _name = Pointer_stringify(name);
-			var gamedata = Module.HEAP8.subarray(indata, indata + insize);
-			// local storage only takes strings, we need to convert string to base64 before storing.
-			var b64encoded = base64EncArr(gamedata);
-			$.jStorage.set(_name, b64encoded);
-			return true;
-	},
-
-
-	UE_LoadGame: function (name, userIndex, outdataptr, outsizeptr){
-		var _name = Pointer_stringify(name);
-		// local storage only takes strings, we need to convert string to base64 before storing.
-		var b64encoded = $.jStorage.get(_name);
-		if (typeof b64encoded == null || b64encoded == null)
-		    return false;
-		var decodedArray = base64DecToArr(b64encoded);
-		// copy back the decoded array.
-		var outdata = Module._malloc(decodedArray.length);
-		// view the allocated data as a HEAP8.
-		var dest = Module.HEAP8.subarray(outdata, outdata + decodedArray.length);
-		// copy back.
-		for (var i = 0; i < decodedArray.length; ++i) {
-			dest[i] = decodedArray[i];
-		}
-		Module.HEAP32[outsizeptr >> 2] = decodedArray.length;
-		Module.HEAP32[outdataptr >> 2] = outdata;
-		return true;
-	},
-
-	UE_DoesSaveGameExist: function (name, userIndex){
-		var _name = Pointer_stringify(name);
-		var b64encoded = $.jStorage.get(_name);
-		if (b64encoded == null)
-			return false;
-		return true;
-	},
-
-    UE_MessageBox: function (type, message, caption ) {
-        // type maps to EAppMsgType::Type
-        if ( type == 0 ){
-            var text = Pointer_stringify(message);
-            if (!confirm(text)) 
-                return 0; 
-        } else {
-            var text = Pointer_stringify(message);
-            alert(text);
-        }
-        return 1; 
-    },
-
-    UE_GetCurrentCultureName: function (address, outsize) {
-        var culture_name = navigator.language || navigator.browserLanguage;
-        if (culture_name.lenght >= outsize) {
-            return 0;
-        }
-        Module.writeAsciiToMemory(culture_name, address);
-        return 1;
-    },
-
-    UE_MakeHTTPDataRequest: function (ctx, url, verb, payload, payloadsize, headers, freeBuffer, onload, onerror, onprogress) {
-	    var _url = Pointer_stringify(url);
-	    var _verb = Pointer_stringify(verb);
-	    var _headers = Pointer_stringify(headers);
-
-	    var xhr = new XMLHttpRequest();
-	    xhr.open(_verb, _url, true);
-	    xhr.responseType = 'arraybuffer';
-
-        // set all headers. 
-	    var _headerArray = _headers.split("%");
-	    for(var headerArrayidx = 0; headerArrayidx < _headerArray.length; headerArrayidx++){
-	        var header = _headerArray[headerArrayidx].split(":");
-	        xhr.setRequestHeader(header[0],header[1]);
-	    }
-
-	    // Onload event handler
-	    xhr.addEventListener('load', function (e) {
-   	        if (xhr.status === 200 || _url.substr(0, 4).toLowerCase() !== "http") {
-       	            var byteArray = new Uint8Array(xhr.response);
-       	            var buffer = _malloc(byteArray.length);
-       	            HEAPU8.set(byteArray, buffer);
-       	            if (onload) 
-       	                    Runtime.dynCall('viii', onload, [ctx, buffer, byteArray.length]);
-       	            if (freeBuffer) 
-       	                    _free(buffer);
-   	        } 
-   	        else{
-   	            if (onerror) Runtime.dynCall('viii', onerror, [ctx, xhr.status, xhr.statusText]);
-   	        }
-	    });
-
-	    // Onerror event handler
-	    xhr.addEventListener('error', function (e) {
-	        if (onerror) Runtime.dynCall('viii', onerror, [ctx, xhr.status, xhr.statusText]);
-	    });
-
-	    // Onprogress event handler
-	    xhr.addEventListener('progress', function (e) {
-	        if (onprogress) Runtime.dynCall('viii', onprogress, [ctx, e.loaded, e.lengthComputable || e.lengthComputable === undefined ? e.total : 0]);
-	    });
-
-	    // Bypass possible browser redirection limit
-	    try {
-	        if (xhr.channel instanceof Ci.nsIHttpChannel)
-	            xhr.channel.redirectionLimit = 0;
-	    } catch (ex) { }
-
-	    if (_verb === "POST") {
-	        var postData = Module.HEAP8.subarray(payload, payload + payloadsize);
-	        xhr.setRequestHeader("Connection", "close");
-	        xhr.send(postData);
-	    } else {
-	        xhr.send(null);
-	    }
+    // Onload event handler
+    xhr.addEventListener('load', function (e) {
+      if (xhr.status === 200 || _url.substr(0, 4).toLowerCase() !== "http") {
+        var byteArray = new Uint8Array(xhr.response);
+        var buffer = _malloc(byteArray.length);
+        HEAPU8.set(byteArray, buffer);
+        if (onload)
+          Runtime.dynCall('viii', onload, [ctx, buffer, byteArray.length]);
+        if (freeBuffer)
+          _free(buffer);
+      } else {
+        if (onerror)
+          Runtime.dynCall('viii', onerror, [ctx, xhr.status, xhr.statusText]);
+      }
+    });
+
+    // Onerror event handler
+    xhr.addEventListener('error', function (e) {
+      if (onerror)
+        Runtime.dynCall('viii', onerror, [ctx, xhr.status, xhr.statusText]);
+    });
+
+    // Onprogress event handler
+    xhr.addEventListener('progress', function (e) {
+      if (onprogress)
+        Runtime.dynCall('viii', onprogress, [ctx, e.loaded, e.lengthComputable || e.lengthComputable === undefined ? e.total : 0]);
+    });
+
+    // Bypass possible browser redirection limit
+    try {
+      if (xhr.channel instanceof Ci.nsIHttpChannel)
+        xhr.channel.redirectionLimit = 0;
+    } catch (ex) { }
+
+    if (_verb === "POST") {
+      var postData = Module.HEAP8.subarray(payload, payload + payloadsize);
+      xhr.setRequestHeader("Connection", "close");
+      xhr.send(postData);
+    } else {
+      xhr.send(null);
     }
+  }
 };
 
 mergeInto(LibraryManager.library, UE_JavaScriptLibary);
