as of Safari 9.0 setRequestHeader's 2nd parameter value cannot have any leading whitespace

diff --git a/Engine/Source/Runtime/HTML5/HTML5JS/Private/HTML5JavaScriptFx.js b/Engine/Source/Runtime/HTML5/HTML5JS/Private/HTML5JavaScriptFx.js
--- a/Engine/Source/Runtime/HTML5/HTML5JS/Private/HTML5JavaScriptFx.js
+++ b/Engine/Source/Runtime/HTML5/HTML5JS/Private/HTML5JavaScriptFx.js
@@ -121,8 +121,9 @@ var UE_JavaScriptLibary = {
     var _headerArray = _headers.split("%");
     for(var headerArrayidx = 0; headerArrayidx < _headerArray.length; headerArrayidx++){
       var header = _headerArray[headerArrayidx].split(":");
-      xhr.setRequestHeader(header[0],header[1]);
-    }
+      // NOTE: as of Safari 9.0 -- no leading whitespace is allowed on setRequestHeader's 2nd parameter: "value"
+      xhr.setRequestHeader(header[0], header[1].trim());
+  }
 
     // Onload event handler
     xhr.addEventListener('load', function (e) {
@@ -160,6 +161,7 @@ var UE_JavaScriptLibary = {
 
     if (_verb === "POST") {
       var postData = Module.HEAP8.subarray(payload, payload + payloadsize);
+      // WARNING: the following errors on Safari showing: "Refused to set unsafe header "Connection"
       xhr.setRequestHeader("Connection", "close");
       xhr.send(postData);
     } else {
