HTML5 - fix UE_MakeHTTPDataRequest() crash

#rn

diff --git a/Engine/Source/Runtime/HTML5/HTML5JS/Private/HTML5JavaScriptFx.js b/Engine/Source/Runtime/HTML5/HTML5JS/Private/HTML5JavaScriptFx.js
--- a/Engine/Source/Runtime/HTML5/HTML5JS/Private/HTML5JavaScriptFx.js
+++ b/Engine/Source/Runtime/HTML5/HTML5JS/Private/HTML5JavaScriptFx.js
@@ -76,8 +76,8 @@ var UE_JavascriptLibrary = {
     var keys = $.jStorage.index();
     for (var i in keys)
     {
-    	if (keys[i] == _name)
-    		return true;
+      if (keys[i] == _name)
+      return true;
     }
     return false;
   },
@@ -128,11 +128,13 @@ var UE_JavascriptLibrary = {
     xhr.responseType = 'arraybuffer';
 
     // set all headers.
-    var _headerArray = _headers.split("%");
-    for(var headerArrayidx = 0; headerArrayidx < _headerArray.length; headerArrayidx++){
-      var header = _headerArray[headerArrayidx].split(":");
-      // NOTE: as of Safari 9.0 -- no leading whitespace is allowed on setRequestHeader's 2nd parameter: "value"
-      xhr.setRequestHeader(header[0], header[1].trim());
+    if (_headers != "" ) {
+      var _headerArray = _headers.split("%");
+      for(var headerArrayidx = 0; headerArrayidx < _headerArray.length; headerArrayidx++){
+        var header = _headerArray[headerArrayidx].split(":");
+        // NOTE: as of Safari 9.0 -- no leading whitespace is allowed on setRequestHeader's 2nd parameter: "value"
+        xhr.setRequestHeader(header[0], header[1].trim());
+      }
     }
 
     // Onload event handler
