HTML5 - keyboard crash fix

emscripten toolchain upgrade depricated a function -- replacing with newer one

this is part 1 of 2 changes that will be going in for this jira,
this CL will fix the crash (part 2 will move mouse and keyboard input code from ApplicationCore to InputCore)

#rnx

diff --git a/Engine/Source/Runtime/ApplicationCore/Private/HTML5/HTML5Application.cpp b/Engine/Source/Runtime/ApplicationCore/Private/HTML5/HTML5Application.cpp
--- a/Engine/Source/Runtime/ApplicationCore/Private/HTML5/HTML5Application.cpp
+++ b/Engine/Source/Runtime/ApplicationCore/Private/HTML5/HTML5Application.cpp
@@ -181,7 +181,7 @@ EM_BOOL key_callback_browser_thread(int eventType, const EmscriptenKeyboardEvent
 	bool UE4CanvasHasFocus = EM_ASM_INT({return document.activeElement === document.body || document.activeElement === Module['canvas']});
 
 	// Let JavaScript page developer hook into keyboard input event handling behavior.
-	int processingBehavior = EM_ASM_INT({ if (Module['UE4_keyEvent']) { return Module['UE4_keyEvent']($0, Pointer_stringify($1), $2, $3, $4); }}, eventType, e->key, e->keyCode, domPhysicalKeyCode, e);
+	int processingBehavior = EM_ASM_INT({ if (Module['UE4_keyEvent']) { return Module['UE4_keyEvent']($0, UTF8ToString($1), $2, $3, $4); }}, eventType, e->key, e->keyCode, domPhysicalKeyCode, e);
 	if (processingBehavior == UE4_INPUT_EVENT_DISCARD_SUPPRESS_ACTION_YES) return EM_TRUE;
 	if (processingBehavior == UE4_INPUT_EVENT_DISCARD_SUPPRESS_ACTION_NO) return EM_FALSE;
 	if (UE4CanvasHasFocus && processingBehavior != UE4_INPUT_EVENT_DISCARD_SUPPRESS_ACTION_DEFAULT)
diff --git a/Engine/Source/Runtime/Core/Private/HTML5/HTML5PlatformProcess.cpp b/Engine/Source/Runtime/Core/Private/HTML5/HTML5PlatformProcess.cpp
--- a/Engine/Source/Runtime/Core/Private/HTML5/HTML5PlatformProcess.cpp
+++ b/Engine/Source/Runtime/Core/Private/HTML5/HTML5PlatformProcess.cpp
@@ -136,7 +136,7 @@ void FHTML5PlatformProcess::LaunchURL(const TCHAR* URL, const TCHAR* Parms, FStr
 	}
 
 	auto TmpURL = StringCast<ANSICHAR>(URL);
-	MAIN_THREAD_EM_ASM({var InUrl = Pointer_stringify($0); console.log("Opening "+InUrl); window.open(InUrl);}, (ANSICHAR*)TmpURL.Get());
+	MAIN_THREAD_EM_ASM({var InUrl = UTF8ToString($0); console.log("Opening "+InUrl); window.open(InUrl);}, (ANSICHAR*)TmpURL.Get());
 }
 
 const TCHAR* FHTML5PlatformProcess::ExecutableName(bool bRemoveExtension)
diff --git a/Engine/Source/Runtime/Core/Public/HTML5/HTML5AssertionMacros.h b/Engine/Source/Runtime/Core/Public/HTML5/HTML5AssertionMacros.h
--- a/Engine/Source/Runtime/Core/Public/HTML5/HTML5AssertionMacros.h
+++ b/Engine/Source/Runtime/Core/Public/HTML5/HTML5AssertionMacros.h
@@ -29,8 +29,8 @@ void emscripten_log(int flags, ...);
 inline void html5_break_msg(const char* msg, const char* file, int line) {
 	EM_ASM(
 	{
-		var InMsg = Pointer_stringify($0);
-		var InFile = Pointer_stringify($1);
+		var InMsg = UTF8ToString($0);
+		var InFile = UTF8ToString($1);
 		alert('Expression ('+InMsg+') failed in '+InFile+':'+$2+'!\nCheck console for details.\n');
 		var callstack = new Error;
 		throw callstack.stack;
diff --git a/Engine/Source/Runtime/Core/Public/HTML5/HTML5PlatformRunnableThread.h b/Engine/Source/Runtime/Core/Public/HTML5/HTML5PlatformRunnableThread.h
--- a/Engine/Source/Runtime/Core/Public/HTML5/HTML5PlatformRunnableThread.h
+++ b/Engine/Source/Runtime/Core/Public/HTML5/HTML5PlatformRunnableThread.h
@@ -36,7 +36,7 @@ EM_ASM_({ console.log("~FHTML5PlatformProcess DESTROY [" + $0 + "]"); }, this);
 		if (rc == 0)
 		{
 			emscripten_set_thread_name(*HandlePtr, Name);
-//EM_ASM_ARGS({ var str = Pointer_stringify($0); console.log("*** CreateThreadWithName [" + str + "]"); }, Name);
+//EM_ASM_ARGS({ var str = UTF8ToString($0); console.log("*** CreateThreadWithName [" + str + "]"); }, Name);
 		}
 		return rc;
 	}
diff --git a/Engine/Source/Runtime/HTML5/HTML5JS/Private/HTML5JavaScriptFx.js b/Engine/Source/Runtime/HTML5/HTML5JS/Private/HTML5JavaScriptFx.js
--- a/Engine/Source/Runtime/HTML5/HTML5JS/Private/HTML5JavaScriptFx.js
+++ b/Engine/Source/Runtime/HTML5/HTML5JS/Private/HTML5JavaScriptFx.js
@@ -6,7 +6,7 @@ var UE_JavascriptLibrary = {
     // NOTE: C++ calling this function are written with syncronus logic
     // otherwise, use: UE_MakeHTTPDataRequest() -- see below
 
-    var _url = Pointer_stringify(url);
+    var _url = UTF8ToString(url);
 
     var request = new XMLHttpRequest();
     request.UE_fetch = {
@@ -38,7 +38,7 @@ var UE_JavascriptLibrary = {
   // ================================================================================
 
   UE_SaveGame: function (name, indata, insize) {
-    var _name = Pointer_stringify(name);
+    var _name = UTF8ToString(name);
     var gamedata = Module.HEAPU8.subarray(indata, indata + insize);
     // local storage only takes strings, we need to convert string to base64 before storing.
     var b64encoded = base64EncArr(gamedata);
@@ -47,7 +47,7 @@ var UE_JavascriptLibrary = {
   },
 
   UE_LoadGame: function (name, outdataptr, outsizeptr) {
-    var _name = Pointer_stringify(name);
+    var _name = UTF8ToString(name);
     // local storage only takes strings, we need to convert string to base64 before storing.
     var b64encoded = $.jStorage.get(_name);
     if (b64encoded === null)
@@ -67,12 +67,12 @@ var UE_JavascriptLibrary = {
   },
 
   UE_DeleteSavedGame: function (name){
-    var _name = Pointer_stringify(name);
+    var _name = UTF8ToString(name);
     return $.jStorage.deleteKey(_name);
   },
 
   UE_DoesSaveGameExist: function (name){
-    var _name = Pointer_stringify(name);
+    var _name = UTF8ToString(name);
     var keys = $.jStorage.index();
     for (var i in keys)
     {
@@ -87,7 +87,7 @@ var UE_JavascriptLibrary = {
 
   UE_MessageBox: function (type, message, caption ) {
     // type maps to EAppMsgType::Type
-    var text = Pointer_stringify(message);
+    var text = UTF8ToString(message);
     if (!type) return confirm(text);
     alert(text);
     return 1;
@@ -113,9 +113,9 @@ var UE_JavascriptLibrary = {
   // ================================================================================
 
   UE_MakeHTTPDataRequest: function (ctx, url, verb, payload, payloadsize, headers, async, freeBuffer, onload, onerror, onprogress) {
-    var _url = Pointer_stringify(url);
-    var _verb = Pointer_stringify(verb);
-    var _headers = Pointer_stringify(headers);
+    var _url = UTF8ToString(url);
+    var _verb = UTF8ToString(verb);
+    var _headers = UTF8ToString(headers);
 
     var xhr = new XMLHttpRequest();
     xhr.UE_fetch = {
