HTML5 - make save game system checks get skipped under multi-threaded builds...

#rn

diff --git a/Engine/Source/Runtime/HTML5/HTML5JS/Private/HTML5JavaScriptFx.js b/Engine/Source/Runtime/HTML5/HTML5JS/Private/HTML5JavaScriptFx.js
--- a/Engine/Source/Runtime/HTML5/HTML5JS/Private/HTML5JavaScriptFx.js
+++ b/Engine/Source/Runtime/HTML5/HTML5JS/Private/HTML5JavaScriptFx.js
@@ -38,6 +38,7 @@ var UE_JavascriptLibrary = {
   // ================================================================================
 
   UE_SaveGame: function (name, indata, insize) {
+    if (ENVIRONMENT_IS_WORKER) { return false; }
     var _name = UTF8ToString(name);
     var gamedata = Module.HEAPU8.subarray(indata, indata + insize);
     // local storage only takes strings, we need to convert string to base64 before storing.
@@ -47,6 +48,7 @@ var UE_JavascriptLibrary = {
   },
 
   UE_LoadGame: function (name, outdataptr, outsizeptr) {
+    if (ENVIRONMENT_IS_WORKER) { return false; }
     var _name = UTF8ToString(name);
     // local storage only takes strings, we need to convert string to base64 before storing.
     var b64encoded = $.jStorage.get(_name);
@@ -67,11 +69,13 @@ var UE_JavascriptLibrary = {
   },
 
   UE_DeleteSavedGame: function (name){
+    if (ENVIRONMENT_IS_WORKER) { return false; }
     var _name = UTF8ToString(name);
     return $.jStorage.deleteKey(_name);
   },
 
   UE_DoesSaveGameExist: function (name){
+    if (ENVIRONMENT_IS_WORKER) { return false; }
     var _name = UTF8ToString(name);
     var keys = $.jStorage.index();
     for (var i in keys)
